Adding Collision tool #rb none

diff --git a/Engine/Plugins/Experimental/InstanceLevelCollision/InstanceLevelCollision.uplugin b/Engine/Plugins/Experimental/InstanceLevelCollision/InstanceLevelCollision.uplugin
new file mode 100644
--- /dev/null
+++ b/Engine/Plugins/Experimental/InstanceLevelCollision/InstanceLevelCollision.uplugin
@@ -0,0 +1,37 @@
+{
+	"FileVersion": 3,
+	"Version": 1,
+	"VersionName": "1.0",
+	"FriendlyName": "InstanceLevelCollision",
+	"Description": "",
+	"Category": "Other",
+	"CreatedBy": "",
+	"CreatedByURL": "",
+	"DocsURL": "",
+	"MarketplaceURL": "",
+	"SupportURL": "",
+	"EngineVersion": "5.0.0",
+	"CanContainContent": true,
+	"Installed": true,
+	"Modules": [
+		{
+			"Name": "InstanceLevelCollision",
+			"Type": "Editor",
+			"LoadingPhase": "PreLoadingScreen"
+		}
+	],
+	"Plugins": [
+		{
+			"Name": "ModelingToolsEditorMode",
+			"Enabled": true
+		},
+		{
+			"Name": "GeometryProcessing",
+			"Enabled": true
+		},
+		{
+			"Name": "MeshModelingToolset",
+			"Enabled": true
+		}
+	]
+}
\ No newline at end of file
diff --git a/Engine/Plugins/Experimental/InstanceLevelCollision/Source/InstanceLevelCollision/InstanceLevelCollision.Build.cs b/Engine/Plugins/Experimental/InstanceLevelCollision/Source/InstanceLevelCollision/InstanceLevelCollision.Build.cs
new file mode 100644
--- /dev/null
+++ b/Engine/Plugins/Experimental/InstanceLevelCollision/Source/InstanceLevelCollision/InstanceLevelCollision.Build.cs
@@ -0,0 +1,76 @@
+// Some copyright should be here...
+
+using UnrealBuildTool;
+
+public class InstanceLevelCollision : ModuleRules
+{
+	public InstanceLevelCollision(ReadOnlyTargetRules Target) : base(Target)
+	{
+		PCHUsage = ModuleRules.PCHUsageMode.UseExplicitOrSharedPCHs;
+
+		if (Target.Type == TargetType.Editor)
+		{
+			PrivateDependencyModuleNames.AddRange(
+				new string[] {
+					"UnrealEd"
+				}
+			);
+		}
+		PublicIncludePaths.AddRange(
+			new string[] {
+				// ... add public include paths required here ...
+			}
+			);
+				
+		
+		PrivateIncludePaths.AddRange(
+			new string[] {
+				// ... add other private include paths required here ...
+			}
+			);
+			
+		
+		PublicDependencyModuleNames.AddRange(
+			new string[]
+			{
+				"Core",
+				// ... add other public dependencies that you statically link with here ...
+			}
+			);
+			
+		
+		PrivateDependencyModuleNames.AddRange(
+			new string[]
+			{
+				"CoreUObject",
+				"Engine",
+				"Slate",
+				"SlateCore",
+				"InteractiveToolsFramework",
+				"DynamicMesh",
+				"MeshConversion",
+				"ModelingOperators",
+				"ModelingOperatorsEditorOnly",
+				"DynamicMesh",
+				"MeshModelingTools",
+				"ModelingComponents",
+				"GeometricObjects",
+				"GeometryAlgorithms",
+				"AssetTools",
+				"MeshUtilitiesCommon",
+				"MeshDescription",
+				"StaticMeshDescription",
+				"MeshReductionInterface"
+				// ... add private dependencies that you statically link with here ...	
+			}
+			);
+		
+		
+		DynamicallyLoadedModuleNames.AddRange(
+			new string[]
+			{
+				// ... add any modules that your module loads dynamically here ...
+			}
+			);
+	}
+}
diff --git a/Engine/Plugins/Experimental/InstanceLevelCollision/Source/InstanceLevelCollision/Private/InstanceLevelCollision.cpp b/Engine/Plugins/Experimental/InstanceLevelCollision/Source/InstanceLevelCollision/Private/InstanceLevelCollision.cpp
new file mode 100644
--- /dev/null
+++ b/Engine/Plugins/Experimental/InstanceLevelCollision/Source/InstanceLevelCollision/Private/InstanceLevelCollision.cpp
@@ -0,0 +1,22 @@
+// Copyright Epic Games, Inc. All Rights Reserved.
+
+#include "InstanceLevelCollision.h"
+
+#define LOCTEXT_NAMESPACE "FInstanceLevelCollisionModule"
+
+void FInstanceLevelCollisionModule::StartupModule()
+{
+	// This code will execute after your module is loaded into memory; the exact timing is specified in the .uplugin file per-module
+	
+}
+
+void FInstanceLevelCollisionModule::ShutdownModule()
+{
+	// This function may be called during shutdown to clean up your module.  For modules that support dynamic reloading,
+	// we call this function before unloading the module.
+	
+}
+
+#undef LOCTEXT_NAMESPACE
+	
+IMPLEMENT_MODULE(FInstanceLevelCollisionModule, InstanceLevelCollision)
\ No newline at end of file
diff --git a/Engine/Plugins/Experimental/InstanceLevelCollision/Source/InstanceLevelCollision/Private/InstanceLevelCollisionBPLibrary.cpp b/Engine/Plugins/Experimental/InstanceLevelCollision/Source/InstanceLevelCollision/Private/InstanceLevelCollisionBPLibrary.cpp
new file mode 100644
--- /dev/null
+++ b/Engine/Plugins/Experimental/InstanceLevelCollision/Source/InstanceLevelCollision/Private/InstanceLevelCollisionBPLibrary.cpp
@@ -0,0 +1,578 @@
+// Copyright Epic Games, Inc. All Rights Reserved.
+
+#include "InstanceLevelCollisionBPLibrary.h"
+#include "Engine/InstancedStaticMesh.h"
+#include "LandscapeComponent.h"
+#include "LandscapeGrassType.h"
+#include "ComponentSourceInterfaces.h" 
+#include "CompositionOps/VoxelSolidifyMeshesOp.h"
+#include "DynamicMesh3.h"
+#include "DynamicMeshToMeshDescription.h"
+#include "MeshDescriptionToDynamicMesh.h"
+#include "CleaningOps/SimplifyMeshOp.h"
+#include "CleaningOps/RemeshMeshOp.h"
+#include "Misc/Paths.h"
+#include "Engine/StaticMesh.h"
+#include "UObject/UObjectGlobals.h"
+#include "InstanceLevelCollision.h"
+#include "ToolSceneQueriesUtil.h"
+#include "CompGeom/PolygonTriangulation.h"
+#include "ConvexHull2.h"
+#include "DynamicMeshEditor.h"
+#include "Generators/SweepGenerator.h"
+#include "AssetToolsModule.h"
+#include "FileHelpers.h"
+#include "IAssetTools.h"
+#include "RemeshMeshTool.h"
+#include "Generators/RectangleMeshGenerator.h"
+#include "LevelInstance/LevelInstanceSubsystem.h"
+#include "LevelInstance/Packed/PackedLevelInstanceActor.h"
+#include "LevelInstance/Packed/PackedLevelInstanceBuilder.h"
+#include "Operations/MergeCoincidentMeshEdges.h"
+#include "CleaningOps/RemoveOccludedTrianglesOp.h"
+#include "CompositionOps/VoxelMorphologyMeshesOp.h"
+#include "Async/Async.h"
+#include "OverlappingCorners.h"
+#include "StaticMeshOperations.h"
+#include "IMeshReductionManagerModule.h"
+#include "IMeshReductionInterfaces.h"
+
+
+
+
+UInstanceLevelCollisionBPLibrary::UInstanceLevelCollisionBPLibrary(const FObjectInitializer& ObjectInitializer)
+: Super(ObjectInitializer)
+{
+
+}
+
+void simplifyMesh(FMeshDescription* OriginalMeshDescription, FProgressCancel* Progress, bool bDiscardAttributes, int TargetCount, FMeshDescription &DstMeshDescription)
+{
+
+	const FMeshDescription* SrcMeshDescription = OriginalMeshDescription;
+	//FMeshDescription DstMeshDescription(*SrcMeshDescription);
+	//DstMeshDescription = *SrcMeshDescription;
+	TUniquePtr<FDynamicMesh3> ResultMesh;
+	if (Progress && Progress->Cancelled())
+	{
+		return;
+	}
+
+	FOverlappingCorners OverlappingCorners;
+	FStaticMeshOperations::FindOverlappingCorners(OverlappingCorners, *SrcMeshDescription, 1.e-5);
+
+	if (Progress && Progress->Cancelled())
+	{
+		return;
+	}
+
+	FMeshReductionSettings ReductionSettings;
+		int32 NumTris = SrcMeshDescription->Polygons().Num();
+		ReductionSettings.PercentTriangles = (float)TargetCount / (float)NumTris;
+		ReductionSettings.TerminationCriterion = EStaticMeshReductionTerimationCriterion::Triangles;
+
+		IMeshReductionManagerModule& MeshReductionModule = FModuleManager::Get().LoadModuleChecked<IMeshReductionManagerModule>("MeshReductionInterface");
+		IMeshReduction* MeshReduction = MeshReductionModule.GetStaticMeshReductionInterface();
+
+	float Error;
+	{
+		/*if (!MeshReduction)
+		{
+			// no reduction possible, failed to load the required interface
+			Error = 0.f;
+			ResultMesh->Copy(*OriginalMesh, true, true, true, !bDiscardAttributes);
+			return;
+		}*/
+
+		Error = ReductionSettings.MaxDeviation;
+		MeshReduction->ReduceMeshDescription(DstMeshDescription, Error, *SrcMeshDescription, OverlappingCorners, ReductionSettings);
+	}
+
+	if (Progress && Progress->Cancelled())
+	{
+		return;
+	}
+	/*
+	// Put the reduced mesh into the target...
+	FMeshDescriptionToDynamicMesh Converter;
+	Converter.Convert(&DstMeshDescription, *ResultMesh);
+	if (bDiscardAttributes)
+	{
+		ResultMesh->DiscardAttributes();
+	}
+
+
+	bool bFailedModifyNeedsRegen = false;
+	// The UE4 tool will split the UV boundaries.  Need to weld this.
+	{
+		FDynamicMesh3* ComponentMesh = ResultMesh.Get();
+
+		FMergeCoincidentMeshEdges Merger(ComponentMesh);
+		Merger.MergeSearchTolerance = 10.0f * FMathf::ZeroTolerance;
+		Merger.OnlyUniquePairs = false;
+		if (Merger.Apply() == false)
+		{
+			bFailedModifyNeedsRegen = true;
+		}
+
+		if (Progress && Progress->Cancelled())
+		{
+			return;
+		}
+
+		if (ResultMesh->CheckValidity(true, EValidityCheckFailMode::ReturnOnly) == false)
+		{
+			bFailedModifyNeedsRegen = true;
+		}
+
+		if (Progress && Progress->Cancelled())
+		{
+			return;
+		}
+
+		// in the fallback case where merge edges failed, give up and reset it to what it was before the attempted merger (w/ split UV boundaries everywhere, oh well)
+		if (bFailedModifyNeedsRegen)
+		{
+			ResultMesh->Clear();
+			Converter.Convert(&DstMeshDescription, *ResultMesh);
+			if (bDiscardAttributes)
+			{
+				ResultMesh->DiscardAttributes();
+			}
+		}
+	}*/
+}
+
+bool CapBottom(FDynamicMesh3* Mesh, FDynamicMesh3& Projected, float &ZValue, float Offset, FTransform Actortransform, ECollisionMaxSlice CollisionType = ECollisionMaxSlice::MinZ,bool bFlatBase = true, bool bMakeBasin = false)
+{
+	// compute the 2D convex hull
+	FConvexHull2d HullCompute;
+	TArray<FVector2d> ProjectedVertices;
+	ProjectedVertices.SetNum(Mesh->MaxVertexID());
+	for (int VID : Mesh->VertexIndicesItr())
+	{
+		const FVector3d& V = Mesh->GetVertexRef(VID);
+		ProjectedVertices[VID] = FVector2d(V.X, V.Y);
+	}
+	bool bOK = HullCompute.Solve(Mesh->MaxVertexID(),
+		[&ProjectedVertices](int VID) { return ProjectedVertices[VID]; },
+		[&Mesh](int VID) { return Mesh->IsVertex(VID); }
+	);
+	if (!bOK)
+	{
+		return false;
+	}
+	// extract polygon
+	const TArray<int32>& PolygonIndices = HullCompute.GetPolygonIndices();
+	TArray<FVector2d> PolygonVertices;  PolygonVertices.SetNum(PolygonIndices.Num());
+	// the min and max Z positions along the outer boundary; could be good reference points for placing the bottom cap
+	double MinZ = FMathd::MaxReal, MaxZ = -FMathd::MaxReal;
+
+	switch (CollisionType)
+	{
+	case ECollisionMaxSlice::XYBound:
+		for (int32 Idx = 0; Idx < PolygonVertices.Num(); Idx++)
+		{
+			PolygonVertices[Idx] = ProjectedVertices[PolygonIndices[Idx]];
+			if (bFlatBase)
+			{
+				double Z = Mesh->GetVertex(PolygonIndices[Idx]).Z;
+				MinZ = FMathd::Min(MinZ, Z);
+				MaxZ = FMathd::Max(MaxZ, Z);
+			}
+		}
+		ZValue = bMakeBasin ? MinZ : MaxZ;
+		ZValue += Offset;
+		break;
+
+	case ECollisionMaxSlice::MinZ:
+	{
+		FAxisAlignedBox3d MeshBound = Mesh->GetCachedBounds();
+		MinZ = MeshBound.Min.Z;
+		ZValue = MinZ + Offset;
+		break;
+	}
+
+	case ECollisionMaxSlice::Custom:
+		ZValue = Actortransform.InverseTransformPosition(FVector::ZeroVector).Z + Offset;
+		break;
+	}
+
+	// triangulate polygon
+	TArray<FIndex3i> Triangles;
+	PolygonTriangulation::TriangulateSimplePolygon(PolygonVertices, Triangles);
+	// fill mesh with result
+	// optionally make an open base enclosing the bottom region by sweeping the convex hull
+	if (bFlatBase && bMakeBasin)
+	{
+		// add sides
+		FGeneralizedCylinderGenerator MeshGen;
+		MeshGen.CrossSection = FPolygon2d(PolygonVertices);
+		MeshGen.Path.Add(FVector3d(0, 0, MinZ));
+		MeshGen.Path.Add(FVector3d(0, 0, MaxZ));
+		MeshGen.bCapped = false;
+		MeshGen.Generate();
+		Projected.Copy(&MeshGen);
+	}
+	else
+	{
+		Projected.Clear();
+	}
+	int StartVID = Projected.MaxVertexID();
+	for (int32 Idx : PolygonIndices)
+	{
+		// either follow the shape of the boundary (with gaps) ...
+		FVector3d Vertex = Mesh->GetVertex(Idx);
+		// ... or make a flat base
+		if (bFlatBase)
+		{
+			Vertex.Z = ZValue;
+		}
+		Projected.AppendVertex(Vertex);
+	}
+	for (const FIndex3i& Tri : Triangles)
+	{
+		Projected.AppendTriangle(FIndex3i(Tri.A + StartVID, Tri.B + StartVID, Tri.C + StartVID));
+	}
+
+	
+	return true;
+}
+
+double CalculateTargetEdgeLength(int TargetTriCount, TSharedPtr<FDynamicMesh3, ESPMode::ThreadSafe> OriginalMesh)
+{
+	double InitialMeshArea = 0;
+	for (int tid : OriginalMesh->TriangleIndicesItr())
+	{
+		InitialMeshArea += OriginalMesh->GetTriArea(tid);
+	}
+
+	double TargetTriArea = InitialMeshArea / (double)TargetTriCount;
+	double EdgeLen = TriangleUtil::EquilateralEdgeLengthForArea(TargetTriArea);
+	return (double)FMath::RoundToInt(EdgeLen * 100.0) / 100.0;
+}
+
+void UInstanceLevelCollisionBPLibrary::GenerateCollision(ALevelInstance* LevelInstance, TArray<AStaticMeshActor*> MeshActor, float ZOffset, ECollisionMaxSlice CollisionType, int VoxelDensity, float TargetPercentage, float Winding)
+{
+	FString LevelName = LevelInstance->GetActorLabel();
+	FString SavePath = FPaths::GetPath(LevelInstance->GetWorldAssetPackage());
+	TArray<TSharedPtr<FDynamicMesh3>> MeshList;
+	TArray<FTransform> ListTransform;
+	FTransform OriginalTransform = LevelInstance->GetTransform();
+
+	if (LevelInstance)
+	{
+		FTransform SavedTransform = LevelInstance->GetTransform();
+		UWorld* World = LevelInstance->GetWorld();
+
+		//Spawn a New Level Instance before the Break
+		APackedLevelInstance* LevelInstanceSpawn = World->SpawnActor<APackedLevelInstance>(LevelInstance->GetClass(), SavedTransform);
+		if (LevelInstance)
+		{
+			TArray<AActor*> BreakActors;
+			LevelInstance->GetLevelInstanceSubsystem()->BreakLevelInstance(LevelInstance, 1U, &BreakActors);
+			for (int i = 0; i < BreakActors.Num(); i++)
+			{
+				//For each Mesh, use the Proxy version 
+				if (AStaticMeshActor* BreakMesh = Cast<AStaticMeshActor>(BreakActors[i]))
+				{
+					FDynamicMesh3 Mesh;
+					bool bMeshIsRealBad = false;
+					FTriMeshCollisionData CollisionData;
+					BreakMesh->GetStaticMeshComponent()->GetStaticMesh()->GetPhysicsTriMeshData(&CollisionData, true);
+					for (FVector V : CollisionData.Vertices)
+					{
+						Mesh.AppendVertex(V);
+					}
+					for (FTriIndices T : CollisionData.Indices)
+					{
+						if (Mesh.FindTriangle(T.v0, T.v1, T.v2) != FDynamicMesh3::InvalidID)
+						{
+							bMeshIsRealBad = true;
+							continue; // skip duplicate triangles in mesh
+						}
+						if (FDynamicMesh3::NonManifoldID == Mesh.AppendTriangle(T.v0, T.v1, T.v2))
+						{
+							int New0 = Mesh.AppendVertex(Mesh, T.v0);
+							int New1 = Mesh.AppendVertex(Mesh, T.v1);
+							int New2 = Mesh.AppendVertex(Mesh, T.v2);
+							Mesh.AppendTriangle(New0, New1, New2);
+							bMeshIsRealBad = true;
+						}
+					}
+					FMergeCoincidentMeshEdges Merger(&Mesh);
+					Merger.Apply();
+
+					//UAddPatchTool PatchTool;
+					//PatchTool->BaseMesh = Mesh;
+
+					MeshList.Add(MakeShared<FDynamicMesh3>(Mesh));
+					MeshActor.Add(BreakMesh);
+					ListTransform.Add(BreakActors[i]->GetActorTransform());
+				}
+			}
+		}
+
+
+		//
+		
+		TArray<TUniquePtr<FPrimitiveComponentTarget>> ComponentTargets;
+		for (int i = 0; i < MeshActor.Num(); i++)
+		{
+			TSet<UActorComponent*> ActorComponent = MeshActor[i]->GetComponents();
+			for (UActorComponent* Components : ActorComponent)
+			{
+				UPrimitiveComponent* MeshComponent = Cast<UPrimitiveComponent>(Components);
+				if (MeshComponent)
+				{
+					
+					ComponentTargets.Add(MakeComponentTarget(MeshComponent));
+				}
+			}
+		}
+
+		FString SubjectName = "test";
+		//Message.Read(TEXT("SubjectName"), SubjectName);
+
+		AsyncTask(ENamedThreads::GameThread, [=]()
+		{
+
+		FProgressCancel Progress;
+
+		//Merge 
+		FDynamicMesh3 MergedMesh;
+		
+		FDynamicMeshEditor MergeEditor(&MergedMesh);
+		FMeshIndexMappings Mappings;
+		for (int32 i = 0; i < MeshList.Num(); i++)
+		{
+			FTransform3d Transform = FTransform3d(ListTransform[i]);
+			if (Transform.GetDeterminant() < 0)
+			{
+				MeshList[i]->ReverseOrientation(false);
+			}
+			MergeEditor.AppendMesh(MeshList[i].Get(), Mappings, [&Transform, &OriginalTransform](int, const FVector3d& P) {return Transform.TransformPosition(P) - OriginalTransform.GetTranslation(); }, [&Transform](int, const FVector3d& N) {return Transform.TransformVector(N); });
+			
+		}
+
+		FMergeCoincidentMeshEdges Merger(&MergedMesh);
+		Merger.Apply();
+
+		FDynamicMesh3 CappedMesh;
+		FDynamicMesh3 Projected;
+		float ZValue = 0;
+		CapBottom(&MergedMesh, Projected, ZValue, ZOffset, OriginalTransform, CollisionType);
+
+		TArray<int> RemoveTris;
+		for (int tid : MergedMesh.TriangleIndicesItr())
+		{
+
+			FIndex3i Tri = MergedMesh.GetTriangle(tid);
+			
+				if (MergedMesh.GetVertex(Tri.A).Z < ZValue)
+				{
+					RemoveTris.Add(tid);
+				}
+		}
+
+		FDynamicMeshEditor Editor(&MergedMesh);
+		Editor.RemoveTriangles(RemoveTris, true);
+
+
+		/*FMeshIndexMappings IndexMaps;
+		Editor.AppendMesh(&Projected, IndexMaps);*/
+
+				//Init Simply Mesh tool
+		/*TUniquePtr<FSimplifyMeshOp> SimplifyOp = MakeUnique<FSimplifyMeshOp>();
+		SimplifyOp->bDiscardAttributes = false;
+		SimplifyOp->bPreventNormalFlips = true;
+		SimplifyOp->bPreserveSharpEdges = true;
+		SimplifyOp->bAllowSeamCollapse = false;
+		SimplifyOp->bReproject = false;
+		SimplifyOp->SimplifierType = ESimplifyType::QEM;
+		SimplifyOp->TargetEdgeLength = 5.0;
+		SimplifyOp->TargetMode = ESimplifyTargetType::Percentage;
+		SimplifyOp->TargetPercentage = 20;
+		SimplifyOp->MeshBoundaryConstraint = EEdgeRefineFlags::NoConstraint;
+		SimplifyOp->GroupBoundaryConstraint = EEdgeRefineFlags::NoConstraint;
+		SimplifyOp->MaterialBoundaryConstraint = EEdgeRefineFlags::NoConstraint;
+		SimplifyOp->OriginalMesh = MakeShared<FDynamicMesh3, ESPMode::ThreadSafe>(MergedMesh);
+		SimplifyOp->OriginalMeshSpatial = MakeShared<FDynamicMeshAABBTree3, ESPMode::ThreadSafe>(SimplifyOp->OriginalMesh.Get());
+		SimplifyOp->CalculateResult(&Progress);
+		TUniquePtr<FDynamicMesh3> SimplifyNewMesh = SimplifyOp->ExtractResult();*/
+
+
+
+
+
+
+
+
+		TUniquePtr<FRemeshMeshOp> RemeshOp = MakeUnique<FRemeshMeshOp>();
+
+		RemeshOp->RemeshType = ERemeshType::Standard;
+		RemeshOp->bCollapses = true;
+		RemeshOp->bDiscardAttributes = false;
+		RemeshOp->bFlips = true;
+		RemeshOp->bPreserveSharpEdges = true;
+		RemeshOp->SmoothingType = ERemeshSmoothingType::MeanValue;
+		RemeshOp->MaxRemeshIterations = 20;
+		RemeshOp->RemeshIterations = 20;
+		RemeshOp->bReproject = true;
+		RemeshOp->ProjectionTarget = &MergedMesh;
+		RemeshOp->MeshBoundaryConstraint = EEdgeRefineFlags::NoConstraint;
+		RemeshOp->GroupBoundaryConstraint = EEdgeRefineFlags::NoConstraint;
+		RemeshOp->MaterialBoundaryConstraint = EEdgeRefineFlags::NoConstraint;
+
+		TUniquePtr<FDynamicMesh3> ProjectionTarget = MakeUnique<FDynamicMesh3>(MergedMesh);
+		TUniquePtr<FDynamicMeshAABBTree3> ProjectionTargetSpatial = MakeUnique<FDynamicMeshAABBTree3>(ProjectionTarget.Get(), true);
+		RemeshOp->ProjectionTargetSpatial = ProjectionTargetSpatial.Get();
+		
+		RemeshOp->OriginalMesh = MakeShared<FDynamicMesh3, ESPMode::ThreadSafe>(MergedMesh);
+
+		//TSharedPtr<FDynamicMeshAABBTree3, ESPMode::ThreadSafe> OriginalMeshSpatial;
+		RemeshOp->OriginalMeshSpatial = MakeShared<FDynamicMeshAABBTree3, ESPMode::ThreadSafe>(RemeshOp->OriginalMesh.Get(), true);
+		RemeshOp->TargetEdgeLength = CalculateTargetEdgeLength(MergedMesh.TriangleCount(), RemeshOp->OriginalMesh);
+
+		
+		RemeshOp->CalculateResult(&Progress);
+		TUniquePtr<FDynamicMesh3> Remesh = RemeshOp->ExtractResult();
+
+		//Jacketing Mesh 
+		TSharedPtr<IndexMeshWithAcceleration, ESPMode::ThreadSafe> CombinedMeshTrees;
+		CombinedMeshTrees = MakeShared<IndexMeshWithAcceleration,ESPMode::ThreadSafe>();
+
+
+		TUniquePtr<FRemoveOccludedTrianglesOp> JacketingOp = MakeUnique<FRemoveOccludedTrianglesOp>();
+		JacketingOp->InsideMode = EOcclusionCalculationMode::SimpleOcclusionTest;
+		JacketingOp->AddTriangleSamples = 4;
+		JacketingOp->AddRandomRays = 4;
+		JacketingOp->MeshTransforms.SetNum(1);
+		CombinedMeshTrees->AddMesh(*Remesh.Get(), FTransform3d::Identity());
+		CombinedMeshTrees->AddMesh(Projected, FTransform3d::Identity());
+		CombinedMeshTrees->BuildAcceleration();
+		JacketingOp->OriginalMesh = MakeShareable<FDynamicMesh3>(Remesh.Release());
+		JacketingOp->CombinedMeshTrees = CombinedMeshTrees;
+		JacketingOp->CalculateResult(&Progress);
+		TUniquePtr<FDynamicMesh3> jacketMesh = JacketingOp->ExtractResult();
+		
+		// Init Vox Wrap tool
+		TUniquePtr<FVoxelSolidifyMeshesOp> Op = MakeUnique<FVoxelSolidifyMeshesOp>();
+		Op->Transforms.SetNum(1);
+		Op->Meshes.SetNum(1);
+		TArray<TSharedPtr<const FDynamicMesh3, ESPMode::ThreadSafe>> OriginalDynamicMeshes;
+		OriginalDynamicMeshes.Add(MakeShareable<FDynamicMesh3>(jacketMesh.Release()));
+		Op->Meshes = OriginalDynamicMeshes;
+		Op->Transforms[0] = FTransform::Identity;
+		Op->OutputVoxelCount = VoxelDensity;
+		Op->InputVoxelCount = VoxelDensity;
+		//Op->bApplyThickenShells
+		Op->bAutoSimplify = false;
+		Op->WindingThreshold = Winding;
+		Op->CalculateResult(&Progress);
+		TUniquePtr<FDynamicMesh3> Newmesh = Op->ExtractResult();
+
+		TUniquePtr<FVoxelMorphologyMeshesOp> MorphOp = MakeUnique<FVoxelMorphologyMeshesOp>();
+		MorphOp->Transforms.SetNum(1);
+		MorphOp->Meshes.SetNum(1);
+		TArray<TSharedPtr<const FDynamicMesh3, ESPMode::ThreadSafe>> OriginalDynamicMeshesMorph;
+		OriginalDynamicMeshesMorph.Add(MakeShareable<FDynamicMesh3>(Newmesh.Release()));
+		//FMeshDescriptionToDynamicMesh Converter;
+		MorphOp->Meshes = OriginalDynamicMeshesMorph;
+		MorphOp->Transforms[0] = FTransform::Identity;
+		MorphOp->OutputVoxelCount = VoxelDensity;
+		MorphOp->InputVoxelCount = VoxelDensity;
+		MorphOp->Operation = EMorphologyOperation::Close;
+		//Op->bApplyThickenShells
+		MorphOp->CalculateResult(&Progress);
+		TUniquePtr<FDynamicMesh3> Morphmesh = MorphOp->ExtractResult();
+		
+
+		
+		//Create the UStaticMesh
+	
+		FString Name = LevelName;
+		const FString DefaultSuffix = TEXT("_Collider");
+		FString PackageName = SavePath +"/Collider/" + Name+ DefaultSuffix;
+		UPackage* Package = CreatePackage(NULL, *PackageName);
+		FAssetToolsModule& AssetToolsModule = FModuleManager::LoadModuleChecked<FAssetToolsModule>("AssetTools");
+		AssetToolsModule.Get().CreateUniqueAssetName(PackageName,TEXT(""), PackageName, Name);
+		FName MeshName = *FPackageName::GetLongPackageAssetName(PackageName);
+		UStaticMesh* myStaticMesh = NewObject<UStaticMesh>(Package, MeshName, RF_Public | RF_Standalone);
+
+		myStaticMesh->InitResources();
+		myStaticMesh->SetNumSourceModels(0);
+		FStaticMeshSourceModel& SrcModel = myStaticMesh->AddSourceModel();
+		FMeshBuildSettings buildsetting;
+		FMeshDescription* MeshDescription = myStaticMesh->CreateMeshDescription(0);
+
+
+		
+		
+
+		FDynamicMeshToMeshDescription Converters;
+		Converters.Convert(Morphmesh.Get(), *MeshDescription);
+		FMeshDescription SimplifyMeshDescription;
+		//simplifyMesh(MeshDescription, &Progress, true, TargetPercentage, SimplifyMeshDescription);
+		//*MeshDescription = SimplifyMeshDescription;
+
+
+		//Simplify Final Mesh
+		TUniquePtr<FSimplifyMeshOp> FinalOp = MakeUnique<FSimplifyMeshOp>();
+		FinalOp->bDiscardAttributes = false;
+		FinalOp->bPreventNormalFlips = true;
+		FinalOp->bPreserveSharpEdges = true;
+		FinalOp->bAllowSeamCollapse = false;
+		FinalOp->bReproject = false;
+		FinalOp->TargetEdgeLength = 5.0;
+		FinalOp->SimplifierType = ESimplifyType::UE4Standard;
+		FinalOp->TargetMode = ESimplifyTargetType::TriangleCount;
+		FinalOp->TargetCount = TargetPercentage;
+		FinalOp->MeshBoundaryConstraint = EEdgeRefineFlags::NoConstraint;
+		FinalOp->GroupBoundaryConstraint = EEdgeRefineFlags::NoConstraint;
+		FinalOp->MaterialBoundaryConstraint = EEdgeRefineFlags::NoConstraint;
+		FinalOp->OriginalMesh = MakeShareable<FDynamicMesh3>(Morphmesh.Release());
+		FinalOp->OriginalMeshSpatial = MakeShared<FDynamicMeshAABBTree3, ESPMode::ThreadSafe>(FinalOp->OriginalMesh.Get());
+		FinalOp->OriginalMeshDescription = MakeShared<FMeshDescription, ESPMode::ThreadSafe>(*MeshDescription);
+		IMeshReductionManagerModule& MeshReductionModule = FModuleManager::Get().LoadModuleChecked<IMeshReductionManagerModule>("MeshReductionInterface");
+		FinalOp->MeshReduction = MeshReductionModule.GetStaticMeshReductionInterface();
+		FinalOp->CalculateResult(&Progress);
+		TUniquePtr<FDynamicMesh3> FinalMesh = FinalOp->ExtractResult();
+
+		Converters.Convert(FinalMesh.Get(), *MeshDescription);
+
+		TArray<const FMeshDescription*> MeshDescriptionPointers;
+		MeshDescriptionPointers.Add(MeshDescription);
+		UStaticMesh::FBuildMeshDescriptionsParams paramsColl;
+		paramsColl.bBuildSimpleCollision = true;
+		myStaticMesh->BuildFromMeshDescriptions(MeshDescriptionPointers, paramsColl);
+		myStaticMesh->GetBodySetup()->CollisionTraceFlag = ECollisionTraceFlag::CTF_UseComplexAsSimple;
+		
+		TArray<UPackage*> SavePackage;
+		SavePackage.Add(myStaticMesh->GetOutermost());
+		FEditorFileUtils::PromptForCheckoutAndSave(SavePackage, true, true);
+		for (int i = 0; i < MeshActor.Num(); i++)
+		{
+			MeshActor[i]->Destroy();
+		}
+
+		//Add infos to the Spawned LevelInstance
+		FActorSpawnParameters param;
+		LevelInstanceSpawn->Edit();
+		LevelInstanceSpawn->Modify();
+		ULevelInstanceSubsystem* subLevel = LevelInstanceSpawn->GetLevelInstanceSubsystem();
+		ULevel* level = subLevel->GetLevelInstanceLevel(LevelInstanceSpawn);
+		UWorld* LIworld = level->GetWorld();
+		param.OverrideLevel = level;
+		param.Name = FName("Collider");
+		AStaticMeshActor* Collider = LIworld->SpawnActor<AStaticMeshActor>(AStaticMeshActor::StaticClass(), OriginalTransform.GetTranslation(), FRotator(0, 0, 0), param);
+		Collider->SetActorHiddenInGame(true);
+		Collider->SetActorLabel("Collider");
+		Collider->GetStaticMeshComponent()->SetStaticMesh(myStaticMesh);
+		Collider->MarkComponentsRenderStateDirty();
+		LevelInstanceSpawn->Commit();
+			});
+	}
+
+}
+
diff --git a/Engine/Plugins/Experimental/InstanceLevelCollision/Source/InstanceLevelCollision/Public/InstanceLevelCollision.h b/Engine/Plugins/Experimental/InstanceLevelCollision/Source/InstanceLevelCollision/Public/InstanceLevelCollision.h
new file mode 100644
--- /dev/null
+++ b/Engine/Plugins/Experimental/InstanceLevelCollision/Source/InstanceLevelCollision/Public/InstanceLevelCollision.h
@@ -0,0 +1,14 @@
+// Copyright Epic Games, Inc. All Rights Reserved.
+
+#pragma once
+
+#include "Modules/ModuleManager.h"
+
+class FInstanceLevelCollisionModule : public IModuleInterface
+{
+public:
+
+	/** IModuleInterface implementation */
+	virtual void StartupModule() override;
+	virtual void ShutdownModule() override;
+};
diff --git a/Engine/Plugins/Experimental/InstanceLevelCollision/Source/InstanceLevelCollision/Public/InstanceLevelCollisionBPLibrary.h b/Engine/Plugins/Experimental/InstanceLevelCollision/Source/InstanceLevelCollision/Public/InstanceLevelCollisionBPLibrary.h
new file mode 100644
--- /dev/null
+++ b/Engine/Plugins/Experimental/InstanceLevelCollision/Source/InstanceLevelCollision/Public/InstanceLevelCollisionBPLibrary.h
@@ -0,0 +1,54 @@
+// Copyright Epic Games, Inc. All Rights Reserved.
+
+#pragma once
+
+#include "Kismet/BlueprintFunctionLibrary.h"
+/*#include "CoreMinimal.h"
+#include "Landscape.h"
+#include "LandscapeProxy.h"
+#include "Components/HierarchicalInstancedStaticMeshComponent.h"*/
+#include "LevelInstance/LevelInstanceActor.h"
+#include "Engine/StaticMeshActor.h"
+#include "Materials/MaterialInterface.h"
+#include "AddPatchTool.h"
+#include "InstanceLevelCollisionBPLibrary.generated.h"
+
+/* 
+*	Function library class.
+*	Each function in it is expected to be static and represents blueprint node that can be called in any blueprint.
+*
+*	When declaring function you can define metadata for the node. Key function specifiers will be BlueprintPure and BlueprintCallable.
+*	BlueprintPure - means the function does not affect the owning object in any way and thus creates a node without Exec pins.
+*	BlueprintCallable - makes a function which can be executed in Blueprints - Thus it has Exec pins.
+*	DisplayName - full name of the node, shown when you mouse over the node and in the blueprint drop down menu.
+*				Its lets you name the node using characters not allowed in C++ function names.
+*	CompactNodeTitle - the word(s) that appear on the node.
+*	Keywords -	the list of keywords that helps you to find node when you search for it using Blueprint drop-down menu. 
+*				Good example is "Print String" node which you can find also by using keyword "log".
+*	Category -	the category your node will be under in the Blueprint drop-down menu.
+*
+*	For more info on custom blueprint nodes visit documentation:
+*	https://wiki.unrealengine.com/Custom_Blueprint_Node_Creation
+*/
+
+UENUM(BlueprintType)
+enum class ECollisionMaxSlice : uint8
+{
+
+	XYBound,
+
+	MinZ,
+
+	Custom
+};
+
+UCLASS()
+class UInstanceLevelCollisionBPLibrary : public UBlueprintFunctionLibrary
+{
+	GENERATED_UCLASS_BODY()
+
+		class IMeshReduction* MeshReduction;
+
+		UFUNCTION(BlueprintCallable, meta = (DisplayName = "GenerateCollision", Keywords = "GenerateCollision"), Category = "MegaAssemblyLib")
+		static void GenerateCollision(ALevelInstance* LevelInstanceBP, TArray<AStaticMeshActor*> MeshActor, float ZOffset, ECollisionMaxSlice CollisionType, int VoxelDensity = 64, float TargetPercentage = 50.0, float Winding = 0.5);
+};
